MagicLeap: Fix ImageTracking doesn't work on consecutive VRPreview launches and Light estimation doesn't work on when ARSession is restarted on the device.
(Submitted by [REDACTED_PERSON])

diff --git a/Engine/Plugins/Lumin/MagicLeap/Source/MagicLeapAR/Private/LuminARImageTracker.cpp b/Engine/Plugins/Lumin/MagicLeap/Source/MagicLeapAR/Private/LuminARImageTracker.cpp
--- a/Engine/Plugins/Lumin/MagicLeap/Source/MagicLeapAR/Private/LuminARImageTracker.cpp
+++ b/Engine/Plugins/Lumin/MagicLeap/Source/MagicLeapAR/Private/LuminARImageTracker.cpp
@@ -26,6 +26,7 @@ void FLuminARImageTracker::CreateEntityTracker()
 void FLuminARImageTracker::DestroyEntityTracker()
 {
 	IMagicLeapImageTrackerModule::Get().DestroyTracker();
+	TrackedTargetNames.Empty();
 }
 
 void FLuminARImageTracker::OnStartGameFrame()
diff --git a/Engine/Plugins/Lumin/MagicLeapLightEstimation/Source/MagicLeapLightEstimation/Private/MagicLeapLightEstimationModule.cpp b/Engine/Plugins/Lumin/MagicLeapLightEstimation/Source/MagicLeapLightEstimation/Private/MagicLeapLightEstimationModule.cpp
--- a/Engine/Plugins/Lumin/MagicLeapLightEstimation/Source/MagicLeapLightEstimation/Private/MagicLeapLightEstimationModule.cpp
+++ b/Engine/Plugins/Lumin/MagicLeapLightEstimation/Source/MagicLeapLightEstimation/Private/MagicLeapLightEstimationModule.cpp
@@ -39,6 +39,7 @@ void FMagicLeapLightEstimationModule::DestroyTracker()
 	{
 		MLResult Result = MLLightingTrackingDestroy(Tracker);
 		UE_CLOG(Result != MLResult_Ok, LogMagicLeapLightEstimation, Error, TEXT("MLLightingTrackingDestroy failed with error %s"), UTF8_TO_TCHAR(MLGetResultString(Result)));
+		Tracker = ML_INVALID_HANDLE;
 	}
 #endif // WITH_MLSDK
 }
